Fixed non-unity build failure

diff --git a/Engine/Plugins/Runtime/OculusLibrary/Source/OculusLibrary/Private/OculusFunctionLibraryPlugin.cpp b/Engine/Plugins/Runtime/OculusLibrary/Source/OculusLibrary/Private/OculusFunctionLibraryPlugin.cpp
--- a/Engine/Plugins/Runtime/OculusLibrary/Source/OculusLibrary/Private/OculusFunctionLibraryPlugin.cpp
+++ b/Engine/Plugins/Runtime/OculusLibrary/Source/OculusLibrary/Private/OculusFunctionLibraryPlugin.cpp
@@ -2,6 +2,7 @@
 
 #include "OculusFunctionLibraryPrivatePCH.h"
 #include "IOculusLibraryPlugin.h"
+#include "OculusFunctionLibrary.h"
 
 class FOculusLibraryPlugin : public IOculusLibraryPlugin
 {
