Merging

//depot/Partners/Oculus/UE4Pure/Engine/Plugins/Runtime/OculusLibrary/Source/OculusLibrary/Public/OculusFunctionLibrary.h

to //depot/UE4/Engine/Plugins/Runtime/OculusLibrary/Source/OculusLibrary/Public/OculusFunctionLibrary.h

diff --git a/Engine/Plugins/Runtime/OculusLibrary/Source/OculusLibrary/Public/OculusFunctionLibrary.h b/Engine/Plugins/Runtime/OculusLibrary/Source/OculusLibrary/Public/OculusFunctionLibrary.h
--- a/Engine/Plugins/Runtime/OculusLibrary/Source/OculusLibrary/Public/OculusFunctionLibrary.h
+++ b/Engine/Plugins/Runtime/OculusLibrary/Source/OculusLibrary/Public/OculusFunctionLibrary.h
@@ -56,8 +56,8 @@ struct FHmdUserProfile
 		PlayerHeight(0.f), EyeHeight(0.f), IPD(0.f), NeckToEyeDistance(FVector2D::ZeroVector) {}
 };
 
-UCLASS(MinimalAPI)
-class UOculusFunctionLibrary : public UBlueprintFunctionLibrary
+UCLASS()
+class OCULUSLIBRARY_API UOculusFunctionLibrary : public UBlueprintFunctionLibrary
 {
 	GENERATED_UCLASS_BODY()
 	
